Removed the hidden flag that was accidentally added to some plugins

This stopped them appearing in the plugin browser

diff --git a/Engine/Plugins/Runtime/Database/ADOSupport/ADOSupport.uplugin b/Engine/Plugins/Runtime/Database/ADOSupport/ADOSupport.uplugin
--- a/Engine/Plugins/Runtime/Database/ADOSupport/ADOSupport.uplugin
+++ b/Engine/Plugins/Runtime/Database/ADOSupport/ADOSupport.uplugin
@@ -12,7 +12,6 @@
 	"SupportURL": "",
 	"EnabledByDefault": false,
 	"CanContainContent": false,
-	"Hidden": true,
 	"IsBetaVersion": false,
 	"Installed": false,
 	"Modules": [
diff --git a/Engine/Plugins/Runtime/Database/RemoteDatabaseSupport/RemoteDatabaseSupport.uplugin b/Engine/Plugins/Runtime/Database/RemoteDatabaseSupport/RemoteDatabaseSupport.uplugin
--- a/Engine/Plugins/Runtime/Database/RemoteDatabaseSupport/RemoteDatabaseSupport.uplugin
+++ b/Engine/Plugins/Runtime/Database/RemoteDatabaseSupport/RemoteDatabaseSupport.uplugin
@@ -12,7 +12,6 @@
 	"SupportURL": "",
 	"EnabledByDefault": false,
 	"CanContainContent": false,
-	"Hidden": true,
 	"IsBetaVersion": false,
 	"Installed": false,
 	"Modules": [
diff --git a/Engine/Plugins/Runtime/Database/SQLiteCore/SQLiteCore.uplugin b/Engine/Plugins/Runtime/Database/SQLiteCore/SQLiteCore.uplugin
--- a/Engine/Plugins/Runtime/Database/SQLiteCore/SQLiteCore.uplugin
+++ b/Engine/Plugins/Runtime/Database/SQLiteCore/SQLiteCore.uplugin
@@ -12,7 +12,6 @@
 	"SupportURL": "",
 	"EnabledByDefault": false,
 	"CanContainContent": false,
-	"Hidden": true,
 	"IsBetaVersion": false,
 	"Installed": false,
 	"Modules": [
diff --git a/Engine/Plugins/Runtime/Database/SQLiteSupport/SQLiteSupport.uplugin b/Engine/Plugins/Runtime/Database/SQLiteSupport/SQLiteSupport.uplugin
--- a/Engine/Plugins/Runtime/Database/SQLiteSupport/SQLiteSupport.uplugin
+++ b/Engine/Plugins/Runtime/Database/SQLiteSupport/SQLiteSupport.uplugin
@@ -12,7 +12,6 @@
 	"SupportURL": "",
 	"EnabledByDefault": false,
 	"CanContainContent": false,
-	"Hidden": true,
 	"IsBetaVersion": false,
 	"Installed": false,
 	"Modules": [
