PR #8660: Fix `bool ExecuteDBProxyCommand()`

diff --git a/Engine/Plugins/Runtime/Database/RemoteDatabaseSupport/Source/RemoteDatabaseSupport/Private/RemoteDatabaseConnection.cpp b/Engine/Plugins/Runtime/Database/RemoteDatabaseSupport/Source/RemoteDatabaseSupport/Private/RemoteDatabaseConnection.cpp
--- a/Engine/Plugins/Runtime/Database/RemoteDatabaseSupport/Source/RemoteDatabaseSupport/Private/RemoteDatabaseConnection.cpp
+++ b/Engine/Plugins/Runtime/Database/RemoteDatabaseSupport/Source/RemoteDatabaseSupport/Private/RemoteDatabaseConnection.cpp
@@ -19,6 +19,7 @@ bool ExecuteDBProxyCommand(FSocket *Socket, const FString& Cmd)
 
 	// convert to network byte ordering. This is important for running on the ps3 and xenon
 	TCHAR *SendBuf = (TCHAR*)FMemory::Malloc(CmdStrLength * sizeof(TCHAR));
+	FMemory::Memcpy(SendBuf, *Cmd, CmdStrLength * sizeof(TCHAR));
 	NETWORK_ORDER_TCHARARRAY(SendBuf);
 
 	int32 BytesSent = 0;
